ExternalGPUStatistics
- Remove unattended check because automated testing does in fact use this, so it's a valid use case

diff --git a/Engine/Plugins/Runtime/ExternalGPUStatistics/Source/ExternalGPUStatistics/Private/ExternalGPUStatistics.cpp b/Engine/Plugins/Runtime/ExternalGPUStatistics/Source/ExternalGPUStatistics/Private/ExternalGPUStatistics.cpp
--- a/Engine/Plugins/Runtime/ExternalGPUStatistics/Source/ExternalGPUStatistics/Private/ExternalGPUStatistics.cpp
+++ b/Engine/Plugins/Runtime/ExternalGPUStatistics/Source/ExternalGPUStatistics/Private/ExternalGPUStatistics.cpp
@@ -348,12 +348,6 @@ void InitGetGPUUsage()
 		UE_LOG(LogRHI, Display, TEXT("FExternalGPUStatistics disabled due to null RHI."));
 		return;
 	}
-	
-	if (FApp::IsUnattended())
-	{
-		UE_LOG(LogRHI, Display, TEXT("FExternalGPUStatistics doesn't matter when in unattended."));
-		return;
-	}
 
 	if (IsRunningCommandlet())
 	{
