NVML
- We were not checking the return code of the function getting process utilization information
- It appears that the sampling will only be valid every .250 seconds due to an internal sampling rate
- Also, because the data could be stale we may have negative numbers when subtracting the per-process data from the current total utilization
* This is handled by a simple if check

diff --git a/Engine/Plugins/Runtime/ExternalGPUStatistics/Source/ExternalGPUStatistics/Private/ExternalGPUStatistics.cpp b/Engine/Plugins/Runtime/ExternalGPUStatistics/Source/ExternalGPUStatistics/Private/ExternalGPUStatistics.cpp
--- a/Engine/Plugins/Runtime/ExternalGPUStatistics/Source/ExternalGPUStatistics/Private/ExternalGPUStatistics.cpp
+++ b/Engine/Plugins/Runtime/ExternalGPUStatistics/Source/ExternalGPUStatistics/Private/ExternalGPUStatistics.cpp
@@ -175,7 +175,7 @@ bool QueryCurrentProcessGPUStatistics(uint32 GPUIndex, float& OutCurrentProcessU
 		Result = Clocker.GetGpuProcessMHz(GPUIndex, FPlatformProcess::GetCurrentProcessId(), LastSeenTimestamp, CurrentProcessUtilization, CurrentProcessMemoryWrites, OutLastSeenTimestamp);
 		if (!Result)
 		{
-			//return false;
+			return false;
 		}
 		OutCurrentProcessUtilization = CurrentProcessUtilization/100.0f;
 		OutCurrentProcessMemoryUsage = CurrentProcessMemoryWrites/100.0f;
@@ -239,13 +239,14 @@ void UpdateGPUUsageCache()
 			
 			// Because we have the per-process information, we can actually show what the rest of GPU is doing
 			NewFractions.ExternalProcessesMHz -= CurrentProcessMHz;
+			if (NewFractions.ExternalProcessesMHz < 0.0f)
+			{
+				NewFractions.ExternalProcessesMHz = 0.0f;
+			}
 			NewFractions.ExternalProcessMemoryUsage -= CurrentProcessMemoryUsage;
 			NewFractions.CurrentProcessMHz = CurrentProcessMHz;
 			NewFractions.CurrentProcessMemoryUsage = CurrentProcessMemoryUsage;
 		}
-		
-		NewFractions.ExternalProcessesMHz = CurrentGPUMhz;
-		NewFractions.ExternalProcessMemoryUsage = CurrentGPUMemoryUsage;
 
 		// If we do not expect the per process data to work, then copy the result into the per process data to make the output to the CSV and graphs in Horde make sense
 		if (!bIsPerProcessDataExpectedToWork)
