Deny building the External GPU Statistics module on the Server because almost always they have no GPU and to avoid a compile error

#Synced-CL 50411828

diff --git a/Engine/Plugins/Runtime/ExternalGPUStatistics/ExternalGPUStatistics.uplugin b/Engine/Plugins/Runtime/ExternalGPUStatistics/ExternalGPUStatistics.uplugin
--- a/Engine/Plugins/Runtime/ExternalGPUStatistics/ExternalGPUStatistics.uplugin
+++ b/Engine/Plugins/Runtime/ExternalGPUStatistics/ExternalGPUStatistics.uplugin
@@ -25,6 +25,7 @@
 			"Type": "Runtime",
 			"LoadingPhase": "PostEngineInit",
 			"TargetConfigurationDenyList": [ "Shipping" ],
+			"TargetDenyList": ["Server"],
 			"PlatformAllowList": [ "Win64", "Linux" ]
 		}
 	]
